StageMonitor: Fix missing property initialization.

diff --git a/Engine/Source/Runtime/VirtualProduction/StageDataCore/Public/StageMessages.h b/Engine/Source/Runtime/VirtualProduction/StageDataCore/Public/StageMessages.h
--- a/Engine/Source/Runtime/VirtualProduction/StageDataCore/Public/StageMessages.h
+++ b/Engine/Source/Runtime/VirtualProduction/StageDataCore/Public/StageMessages.h
@@ -241,7 +241,7 @@ public:
 
 	/** Event for this critical state */
 	UPROPERTY(VisibleAnywhere, Category = "CriticalState")
-	EStageCriticalStateEvent State;
+	EStageCriticalStateEvent State = EStageCriticalStateEvent::Enter;
 
 	/** Source of the critical state. i.e. TakeName, CustomRecorder, etc... */
 	UPROPERTY(VisibleAnywhere, Category = "CriticalState")
