Fix intermittent test failures caused by FStageProviderMessage initializing DateTime and FrameTime to TimeNow() and GetCurrentFrameTime

The FAutomationTestAttemptToFindUninitializedScriptStructMembers sometime initializes the structure at different times which would fail the test.
Instead of initializing FStageProviderMessage with FDateTime::Now, we are now setting the time in IStageDataProvider::SendMessage which all messages go through.

diff --git a/Engine/Source/Runtime/VirtualProduction/StageDataCore/Private/StageMessages.cpp b/Engine/Source/Runtime/VirtualProduction/StageDataCore/Private/StageMessages.cpp
--- a/Engine/Source/Runtime/VirtualProduction/StageDataCore/Private/StageMessages.cpp
+++ b/Engine/Source/Runtime/VirtualProduction/StageDataCore/Private/StageMessages.cpp
@@ -4,28 +4,6 @@
 
 #include "Misc/App.h"
 
-
-namespace StageProviderMessageUtils
-{
-	static const FQualifiedFrameTime InvalidTime = FQualifiedFrameTime(FFrameTime(FFrameNumber(-1)), FFrameRate(-1, -1));
-}
-
-FStageProviderMessage::FStageProviderMessage()
-{
-	//Common setup of timecode for all provider messages
-	TOptional<FQualifiedFrameTime> CurrentFrameTime = FApp::GetCurrentFrameTime();
-	DateTime = FDateTime::Now();
-	
-	if (CurrentFrameTime.IsSet())
-	{
-		FrameTime = CurrentFrameTime.GetValue();
-	}
-	else
-	{
-		FrameTime = StageProviderMessageUtils::InvalidTime;
-	}
-}
-
 FString FCriticalStateProviderMessage::ToString() const
 {
 	switch (State)
diff --git a/Engine/Source/Runtime/VirtualProduction/StageDataCore/Public/IStageDataProvider.h b/Engine/Source/Runtime/VirtualProduction/StageDataCore/Public/IStageDataProvider.h
--- a/Engine/Source/Runtime/VirtualProduction/StageDataCore/Public/IStageDataProvider.h
+++ b/Engine/Source/Runtime/VirtualProduction/StageDataCore/Public/IStageDataProvider.h
@@ -6,9 +6,13 @@
 #include "Features/IModularFeatures.h"
 
 #include "CoreMinimal.h"
+#include "Misc/App.h"
 #include "StageMessages.h"
 
-
+namespace UE::StageProviderUtils
+{
+	static const FQualifiedFrameTime InvalidTime = FQualifiedFrameTime(FFrameTime(FFrameNumber(-1)), FFrameRate(-1, -1));
+}
 
 struct FStageDataBaseMessage;
 
@@ -74,6 +78,20 @@ private:
 			IStageDataProvider* StageDataProvider = &IModularFeatures::Get().GetModularFeature<IStageDataProvider>(IStageDataProvider::ModularFeatureName);
 
 			static_assert(TIsDerivedFrom<MessageType, FStageProviderEventMessage>::IsDerived || TIsDerivedFrom<MessageType, FStageProviderPeriodicMessage>::IsDerived, "MessageType must be a FStageProviderEventMessage or FStageProviderPeriodicMessage derived UStruct.");
+
+			// Timestamp the outgoing message.
+			TOptional<FQualifiedFrameTime> CurrentFrameTime = FApp::GetCurrentFrameTime();
+			Message.DateTime = FDateTime::Now();
+
+			if (CurrentFrameTime.IsSet())
+			{
+				Message.FrameTime = CurrentFrameTime.GetValue();
+			}
+			else
+			{
+				Message.FrameTime = UE::StageProviderUtils::InvalidTime;
+			}
+
 			return StageDataProvider->SendMessageInternal(&Message, MessageType::StaticStruct(), InFlags);
 		}
 
diff --git a/Engine/Source/Runtime/VirtualProduction/StageDataCore/Public/StageMessages.h b/Engine/Source/Runtime/VirtualProduction/StageDataCore/Public/StageMessages.h
--- a/Engine/Source/Runtime/VirtualProduction/StageDataCore/Public/StageMessages.h
+++ b/Engine/Source/Runtime/VirtualProduction/StageDataCore/Public/StageMessages.h
@@ -175,8 +175,6 @@ struct STAGEDATACORE_API FStageProviderMessage : public FStageDataBaseMessage
 	GENERATED_BODY()
 
 public:
-	FStageProviderMessage();
-
 	/** Method to override if a detailed description is desired in the monitor event viewer */
 	virtual FString ToString() const { return FString(); }
 
