Fix undefined symbol errors when compiling monolithic editor with TimedDataMonitor plugin enabled

#Synced-CL 48060881

diff --git a/Engine/Plugins/VirtualProduction/TimedDataMonitor/Source/TimedDataMonitor/Public/TimedDataMonitorCalibration.h b/Engine/Plugins/VirtualProduction/TimedDataMonitor/Source/TimedDataMonitor/Public/TimedDataMonitorCalibration.h
--- a/Engine/Plugins/VirtualProduction/TimedDataMonitor/Source/TimedDataMonitor/Public/TimedDataMonitorCalibration.h
+++ b/Engine/Plugins/VirtualProduction/TimedDataMonitor/Source/TimedDataMonitor/Public/TimedDataMonitorCalibration.h
@@ -12,6 +12,7 @@
 
 #include "TimedDataMonitorCalibration.generated.h"
 
+#define UE_API TIMEDDATAMONITOR_API
 
 UENUM()
 enum class ETimedDataMonitorCalibrationReturnCode : uint8
@@ -154,27 +155,29 @@ struct FTimedDataMonitorTimeCorrectionParameters
 /**
  * 
  */
-class TIMEDDATAMONITOR_API FTimedDataMonitorCalibration
+class FTimedDataMonitorCalibration
 {
 public:
 	DECLARE_DELEGATE_OneParam(FOnCalibrationCompletedSignature, FTimedDataMonitorCalibrationResult);
 
 public:
-	~FTimedDataMonitorCalibration();
+	UE_API ~FTimedDataMonitorCalibration();
 
-	void Reset();
-	bool IsRunning() const;
-	void CalibrateWithTimecode(const FTimedDataMonitorCalibrationParameters& Parameters, FOnCalibrationCompletedSignature OnCompleted);
+	UE_API void Reset();
+	UE_API bool IsRunning() const;
+	UE_API void CalibrateWithTimecode(const FTimedDataMonitorCalibrationParameters& Parameters, FOnCalibrationCompletedSignature OnCompleted);
 
-	static FTimedDataMonitorCalibrationResult CalibrateWithTimecode(const FTimedDataMonitorCalibrationParameters& CalibrationParameters);
-	static FTimedDataMonitorTimeCorrectionResult ApplyTimeCorrection(const FTimedDataMonitorInputIdentifier& Identifier, const FTimedDataMonitorTimeCorrectionParameters& TimeCorrectionParameters);
+	UE_API static FTimedDataMonitorCalibrationResult CalibrateWithTimecode(const FTimedDataMonitorCalibrationParameters& CalibrationParameters);
+	UE_API static FTimedDataMonitorTimeCorrectionResult ApplyTimeCorrection(const FTimedDataMonitorInputIdentifier& Identifier, const FTimedDataMonitorTimeCorrectionParameters& TimeCorrectionParameters);
 
 private:
-	bool RunCalibrateWithTimecode(float);
-	bool RunCalibrateWithTimecode_CheckForReset(float);
+	UE_API bool RunCalibrateWithTimecode(float);
+	UE_API bool RunCalibrateWithTimecode_CheckForReset(float);
 
 	int32 NumberOfRetry = 0;
 	FTimedDataMonitorCalibrationParameters CalibrationParameters;
 	FTSTicker::FDelegateHandle TickerHandle;
 	FOnCalibrationCompletedSignature OnCalibrationCompleted;
 };
+
+#undef UE_API
\ No newline at end of file
diff --git a/Engine/Plugins/VirtualProduction/TimedDataMonitor/Source/TimedDataMonitor/Public/TimedDataMonitorTypes.h b/Engine/Plugins/VirtualProduction/TimedDataMonitor/Source/TimedDataMonitor/Public/TimedDataMonitorTypes.h
--- a/Engine/Plugins/VirtualProduction/TimedDataMonitor/Source/TimedDataMonitor/Public/TimedDataMonitorTypes.h
+++ b/Engine/Plugins/VirtualProduction/TimedDataMonitor/Source/TimedDataMonitor/Public/TimedDataMonitorTypes.h
@@ -7,9 +7,11 @@
 
 #include "TimedDataMonitorTypes.generated.h"
 
+#define UE_API TIMEDDATAMONITOR_API
 
-USTRUCT(BlueprintType)
-struct TIMEDDATAMONITOR_API FTimedDataMonitorInputIdentifier
+
+USTRUCT(MinimalAPI, BlueprintType)
+struct FTimedDataMonitorInputIdentifier
 {
 	GENERATED_BODY()
 
@@ -18,7 +20,7 @@ private:
 	FGuid Identifier;
 
 public:
-	static FTimedDataMonitorInputIdentifier NewIdentifier();
+	UE_API static FTimedDataMonitorInputIdentifier NewIdentifier();
 
 	bool IsValid() const { return Identifier.IsValid(); }
 
@@ -30,8 +32,8 @@ public:
 };
 
 
-USTRUCT(BlueprintType)
-struct TIMEDDATAMONITOR_API FTimedDataMonitorChannelIdentifier
+USTRUCT(MinimalAPI, BlueprintType)
+struct FTimedDataMonitorChannelIdentifier
 {
 	GENERATED_BODY()
 
@@ -40,7 +42,7 @@ private:
 	FGuid Identifier;
 
 public:
-	static FTimedDataMonitorChannelIdentifier NewIdentifier();
+	UE_API static FTimedDataMonitorChannelIdentifier NewIdentifier();
 
 	bool IsValid() const { return Identifier.IsValid(); }
 
@@ -50,3 +52,5 @@ public:
 
 	friend uint32 GetTypeHash(const FTimedDataMonitorChannelIdentifier& InIdentifier) { return GetTypeHash(InIdentifier.Identifier); }
 };
+
+#undef UE_API
\ No newline at end of file
diff --git a/Engine/Source/Runtime/VirtualProduction/StageDataCore/Public/StageMessages.h b/Engine/Source/Runtime/VirtualProduction/StageDataCore/Public/StageMessages.h
--- a/Engine/Source/Runtime/VirtualProduction/StageDataCore/Public/StageMessages.h
+++ b/Engine/Source/Runtime/VirtualProduction/StageDataCore/Public/StageMessages.h
@@ -18,6 +18,8 @@
 /** Global define to disable any data provider sending logic */
 #define ENABLE_STAGEMONITOR_LOGGING (1 && !NO_LOGGING && !(UE_BUILD_SHIPPING || UE_BUILD_TEST))
 
+#define UE_API STAGEDATACORE_API
+
 
 /** Message flags configuring how a message is sent */
 UENUM()
@@ -68,8 +70,8 @@ enum class EStageLoadingState : uint8
  * Information that won't change for a session
  * Used when monitor and provider connects
  */
-USTRUCT()
-struct STAGEDATACORE_API FStageInstanceDescriptor
+USTRUCT(MinimalAPI)
+struct FStageInstanceDescriptor
 {
 	GENERATED_BODY()
 
@@ -101,13 +103,13 @@ public:
 /**
  *  Base structure for all stage monitoring messages
  */
-USTRUCT()
-struct STAGEDATACORE_API FStageDataBaseMessage
+USTRUCT(MinimalAPI)
+struct FStageDataBaseMessage
 {
 	GENERATED_BODY()
 
 public:
-	virtual ~FStageDataBaseMessage() = default;
+	virtual UE_API ~FStageDataBaseMessage() = default;
 	
 public:
 
@@ -130,8 +132,8 @@ public:
 /**
  *  Base structure for all monitor messages
  */
-USTRUCT()
-struct STAGEDATACORE_API FStageMonitorBaseMessage : public FStageDataBaseMessage
+USTRUCT(MinimalAPI)
+struct FStageMonitorBaseMessage : public FStageDataBaseMessage
 {
 	GENERATED_BODY()
 };
@@ -139,8 +141,8 @@ struct STAGEDATACORE_API FStageMonitorBaseMessage : public FStageDataBaseMessage
 /**
  *  Message broadcasted periodically by the monitor to discover new providers 
  */
-USTRUCT()
-struct STAGEDATACORE_API FStageProviderDiscoveryMessage : public FStageMonitorBaseMessage
+USTRUCT(MinimalAPI)
+struct FStageProviderDiscoveryMessage : public FStageMonitorBaseMessage
 {
 	GENERATED_BODY()
 
@@ -160,8 +162,8 @@ public:
 /**
  *  Message sent when monitor is going down to let know linked providers
  */
-USTRUCT()
-struct STAGEDATACORE_API FStageMonitorCloseMessage : public FStageMonitorBaseMessage
+USTRUCT(MinimalAPI)
+struct FStageMonitorCloseMessage : public FStageMonitorBaseMessage
 {
 	GENERATED_BODY()
 };
@@ -169,8 +171,8 @@ struct STAGEDATACORE_API FStageMonitorCloseMessage : public FStageMonitorBaseMes
 /** 
  * Base Provider messages listened by monitors
  */
-USTRUCT()
-struct STAGEDATACORE_API FStageProviderMessage : public FStageDataBaseMessage
+USTRUCT(MinimalAPI)
+struct FStageProviderMessage : public FStageDataBaseMessage
 {
 	GENERATED_BODY()
 
@@ -192,8 +194,8 @@ public:
 /**
  * Base Provider messages that are events
  */
-USTRUCT()
-struct STAGEDATACORE_API FStageProviderEventMessage : public FStageProviderMessage
+USTRUCT(MinimalAPI)
+struct FStageProviderEventMessage : public FStageProviderMessage
 {
 	GENERATED_BODY()
 };
@@ -201,8 +203,8 @@ struct STAGEDATACORE_API FStageProviderEventMessage : public FStageProviderMessa
 /**
  * Base Provider messages that are periodic
  */
-USTRUCT()
-struct STAGEDATACORE_API FStageProviderPeriodicMessage : public FStageProviderMessage
+USTRUCT(MinimalAPI)
+struct FStageProviderPeriodicMessage : public FStageProviderMessage
 {
 	GENERATED_BODY()
 };
@@ -210,8 +212,8 @@ struct STAGEDATACORE_API FStageProviderPeriodicMessage : public FStageProviderMe
 /**
  *  Message sent by Providers to notify monitors they are closing down
  */
-USTRUCT()
-struct STAGEDATACORE_API FStageProviderCloseMessage : public FStageProviderEventMessage
+USTRUCT(MinimalAPI)
+struct FStageProviderCloseMessage : public FStageProviderEventMessage
 {
 	GENERATED_BODY()
 };
@@ -219,14 +221,14 @@ struct STAGEDATACORE_API FStageProviderCloseMessage : public FStageProviderEvent
 /**
  * Response to a received discovery message sent by providers
  */
-USTRUCT()
-struct STAGEDATACORE_API FStageProviderDiscoveryResponseMessage : public FStageProviderEventMessage
+USTRUCT(MinimalAPI)
+struct FStageProviderDiscoveryResponseMessage : public FStageProviderEventMessage
 {
 	GENERATED_BODY()
 	
 public:
-	FStageProviderDiscoveryResponseMessage() = default;
-	FStageProviderDiscoveryResponseMessage(FStageInstanceDescriptor&& InDescriptor)
+	UE_API FStageProviderDiscoveryResponseMessage() = default;
+	UE_API FStageProviderDiscoveryResponseMessage(FStageInstanceDescriptor&& InDescriptor)
 		: Descriptor(MoveTemp(InDescriptor))
 		{}
 
@@ -241,8 +243,8 @@ public:
 /**
  * Message sent to notify about critical state change.
  */
-USTRUCT()
-struct STAGEDATACORE_API FCriticalStateProviderMessage : public FStageProviderEventMessage
+USTRUCT(MinimalAPI)
+struct FCriticalStateProviderMessage : public FStageProviderEventMessage
 {
 	GENERATED_BODY()
 
@@ -255,7 +257,7 @@ public:
 	{
 	}
 
-	virtual FString ToString() const override;
+	UE_API virtual FString ToString() const override;
 
 	/** Event for this critical state */
 	UPROPERTY(VisibleAnywhere, Category = "CriticalState")
@@ -269,8 +271,8 @@ public:
 /**
  * Message sent to indicate that the node has entered or exited a loading state.
  */
-USTRUCT()
-struct STAGEDATACORE_API FAssetLoadingStateProviderMessage : public FStageProviderEventMessage
+USTRUCT(MinimalAPI)
+struct FAssetLoadingStateProviderMessage : public FStageProviderEventMessage
 {
 	GENERATED_BODY()
 
@@ -281,7 +283,7 @@ public:
 		: LoadingState(InState), AssetName(InAssetName)
 	{}
 
-	virtual FString ToString() const override;
+	UE_API virtual FString ToString() const override;
 
 	/** Event for this critical state */
 	UPROPERTY(VisibleAnywhere, Category = "AssetLoading")
@@ -291,3 +293,5 @@ public:
 	UPROPERTY(VisibleAnywhere, Category = "AssetLoading")
 	FString AssetName;
 };
+ 
+#undef UE_API
\ No newline at end of file
